Adding AutomationUtils experimental plugin, which will house various optional gameplay automation utilities and supporting tools
This initial commit adds:
 - A BlueprintFunctionLibrary that houses a command to take an AutomationScreenshot at gameplay, without having to have FTest running
 - A Commandlet that will do AutomationScreenshot comparisons off-line: this will facilitate taking screenshots on whatever platform client, then doing the image diffs in batch after those are complete.

#rnx

diff --git a/Engine/Plugins/Experimental/AutomationUtils/AutomationUtils.uplugin b/Engine/Plugins/Experimental/AutomationUtils/AutomationUtils.uplugin
new file mode 100644
--- /dev/null
+++ b/Engine/Plugins/Experimental/AutomationUtils/AutomationUtils.uplugin
@@ -0,0 +1,29 @@
+{
+	"FileVersion" : 1,
+	"Version" : 1,
+	"VersionName" : "0.1",
+	"FriendlyName" : "Automation Utilities",
+	"Description" : "Tools and Utilities for Automation purposes",
+	"Category" : "Automation",
+	"CreatedBy" : "Epic Games, Inc.",
+	"CreatedByURL" : "http://epicgames.com",
+	"MarketplaceURL" : "",
+	"SupportURL" : "",
+	"EnabledByDefault" : true,
+	"CanContainContent" : true,
+	"IsBetaVersion" : true,
+	"Installed" : false,
+	"Modules" :
+	[
+		{
+			"Name" : "AutomationUtils",
+			"Type" : "Runtime",
+			"LoadingPhase" : "Default"
+		},
+		{
+			"Name" : "AutomationUtilsEditor",
+			"Type" : "Editor",
+			"LoadingPhase" : "Default"
+		}
+	]
+}
\ No newline at end of file
diff --git a/Engine/Plugins/Experimental/AutomationUtils/Source/AutomationUtils/AutomationUtils.Build.cs b/Engine/Plugins/Experimental/AutomationUtils/Source/AutomationUtils/AutomationUtils.Build.cs
new file mode 100644
--- /dev/null
+++ b/Engine/Plugins/Experimental/AutomationUtils/Source/AutomationUtils/AutomationUtils.Build.cs
@@ -0,0 +1,31 @@
+// Copyright 1998-2019 Epic Games, Inc. All Rights Reserved.
+
+namespace UnrealBuildTool.Rules
+{
+	public class AutomationUtils : ModuleRules
+	{
+		public AutomationUtils(ReadOnlyTargetRules Target) : base(Target)
+		{
+			PublicDependencyModuleNames.AddRange(
+				new string[]
+				{
+					"Core",
+					"CoreUObject",
+					"Engine",
+				}
+			);
+
+			PrivateDependencyModuleNames.AddRange(
+				new string[]
+				{
+					"Core",
+					"CoreUObject",
+					"Engine",
+					"Json",
+					"RHI",
+					"RenderCore"
+                }
+			);
+		}
+	}
+}
diff --git a/Engine/Plugins/Experimental/AutomationUtils/Source/AutomationUtils/Private/AutomationUtils.cpp b/Engine/Plugins/Experimental/AutomationUtils/Source/AutomationUtils/Private/AutomationUtils.cpp
new file mode 100644
--- /dev/null
+++ b/Engine/Plugins/Experimental/AutomationUtils/Source/AutomationUtils/Private/AutomationUtils.cpp
@@ -0,0 +1,7 @@
+// Copyright 1998-2019 Epic Games, Inc. All Rights Reserved.
+
+#include "Modules/ModuleManager.h"
+
+//Module Implementation Call
+IMPLEMENT_MODULE(FDefaultModuleImpl, AutomationUtils)
+
diff --git a/Engine/Plugins/Experimental/AutomationUtils/Source/AutomationUtils/Private/AutomationUtilsBlueprintLibrary.cpp b/Engine/Plugins/Experimental/AutomationUtils/Source/AutomationUtils/Private/AutomationUtilsBlueprintLibrary.cpp
new file mode 100644
--- /dev/null
+++ b/Engine/Plugins/Experimental/AutomationUtils/Source/AutomationUtils/Private/AutomationUtilsBlueprintLibrary.cpp
@@ -0,0 +1,284 @@
+// Copyright 1998-2019 Epic Games, Inc. All Rights Reserved.
+
+#include "AutomationUtilsBlueprintLibrary.h"
+#include "SceneViewExtension.h"
+#include "SceneView.h"
+#include "ShaderCompiler.h"
+#include "ContentStreaming.h"
+#include "ImageUtils.h"
+#include "Misc/EngineVersion.h"
+#include "HardwareInfo.h"
+#include "EngineUtils.h"
+#include "Materials/MaterialInstanceDynamic.h"
+#include "MeshDrawShaderBindings.h"
+#include "Scalability.h"
+
+
+
+
+//Private Helper Class Definitions
+class FAutomationUtilsGameplayViewExtension : public FSceneViewExtensionBase
+{
+public:
+	FAutomationUtilsGameplayViewExtension(const FAutoRegister& AutoRegister)
+		: FSceneViewExtensionBase(AutoRegister)
+	{
+	}
+
+	void SetupViewFamily(FSceneViewFamily& InViewFamily) override
+	{
+		// Turn off common show flags for noisy sources of rendering.
+		FEngineShowFlags& ShowFlags = InViewFamily.EngineShowFlags;
+		ShowFlags.SetAntiAliasing(0);
+		ShowFlags.SetMotionBlur(0);
+		ShowFlags.SetTemporalAA(0);
+		ShowFlags.SetScreenSpaceReflections(0);
+		ShowFlags.SetScreenSpaceAO(0);
+		ShowFlags.SetDistanceFieldAO(0);
+		ShowFlags.SetContactShadows(0);
+		ShowFlags.SetEyeAdaptation(0);
+		ShowFlags.SetBloom(0);
+
+		// Turn off time the ultimate source of noise.
+		InViewFamily.CurrentWorldTime = 0;
+		InViewFamily.CurrentRealTime = 0;
+		InViewFamily.DeltaWorldTime = 0;
+	}
+
+	bool IsActiveThisFrame(class FViewport* InViewport) const
+	{
+		return true;
+	}
+
+	void SetupView(FSceneViewFamily& InViewFamily, FSceneView& InView) override {}
+	void BeginRenderViewFamily(FSceneViewFamily& InViewFamily) override {}
+	void PreRenderViewFamily_RenderThread(FRHICommandListImmediate& RHICmdList, FSceneViewFamily& InViewFamily) override {}
+	void PreRenderView_RenderThread(FRHICommandListImmediate& RHICmdList, FSceneView& InView) override {}
+
+	/** We always want to go last. */
+	virtual int32 GetPriority() const override { return MIN_int32; }
+};
+
+
+class FAutomationUtilsGameplayAutomationScreenshotInstance
+{
+public:
+	FAutomationUtilsGameplayAutomationScreenshotInstance(FString InScreenshotName, float MaxGlobalError, float MaxLocalError)
+		: ScreenshotName(InScreenshotName)
+		, MetadataJsonString(TEXT("{}"))
+		, DeterminedPath(TEXT(""))
+		, World(GWorld.GetReference())
+	{
+		if (GEngine && GEngine->GameViewport)
+		{
+			//FlushRendering
+			FlushRenderingCommands();
+			//hook to the screenshot delegate
+			GEngine->GameViewport->OnScreenshotCaptured().AddRaw(this, &FAutomationUtilsGameplayAutomationScreenshotInstance::HandleScreenshotData);
+			//And a removed-from-world delegate too, just in case
+			FWorldDelegates::LevelRemovedFromWorld.AddRaw(this, &FAutomationUtilsGameplayAutomationScreenshotInstance::WorldDestroyed);
+
+
+
+			//Generate Json Metadata relevant to rendering device, quality settings, and comparison tolerances
+			TSharedPtr<FJsonObject> JsonObject = MakeShareable(new FJsonObject);
+			//General Stuff
+			JsonObject->SetStringField(TEXT("name"), FPaths::MakeValidFileName(InScreenshotName, TEXT('_')));
+			JsonObject->SetStringField(TEXT("context"), GWorld->GetName());
+			JsonObject->SetStringField(TEXT("id"), FGuid::NewGuid().ToString());
+			JsonObject->SetStringField(TEXT("Commit"), FEngineVersion::Current().HasChangelist() ? FString::FromInt(FEngineVersion::Current().GetChangelist()) : FString(TEXT("")));
+			FVector2D ViewportSize;  //Width and Height
+			GEngine->GameViewport->GetViewportSize(ViewportSize);
+			JsonObject->SetNumberField(TEXT("width"), ViewportSize.X);
+			JsonObject->SetNumberField(TEXT("height"), ViewportSize.Y);
+			//RHI
+			JsonObject->SetStringField(TEXT("platform"), FPlatformProperties::IniPlatformName());
+			JsonObject->SetStringField(TEXT("rhi"), FHardwareInfo::GetHardwareInfo(NAME_RHI));
+			FString FeatureLevel;
+			GetFeatureLevelName(GMaxRHIFeatureLevel, FeatureLevel);
+			JsonObject->SetStringField(TEXT("featureLevel"), FeatureLevel);
+			JsonObject->SetBoolField(TEXT("bIsStereo"), GEngine->StereoRenderingDevice.IsValid() ? GEngine->StereoRenderingDevice->IsStereoEnabled() : false);
+			//Vendor
+			JsonObject->SetStringField(TEXT("vendor"), RHIVendorIdToString());
+			JsonObject->SetStringField(TEXT("adapterName"), GRHIAdapterName);
+			JsonObject->SetStringField(TEXT("adapterInternalDriverVersion"), GRHIAdapterInternalDriverVersion);
+			JsonObject->SetStringField(TEXT("adapterUserDriverVersion"), GRHIAdapterUserDriverVersion);
+			JsonObject->SetStringField(TEXT("uniqueDeviceId"), FPlatformMisc::GetDeviceId());
+			//Quality
+			Scalability::FQualityLevels QualityLevels = Scalability::GetQualityLevels();
+			JsonObject->SetNumberField(TEXT("resolutionQuality"), QualityLevels.ResolutionQuality);
+			JsonObject->SetNumberField(TEXT("viewDistanceQuality"), QualityLevels.ViewDistanceQuality);
+			JsonObject->SetNumberField(TEXT("antiAliasingQuality"), QualityLevels.AntiAliasingQuality);
+			JsonObject->SetNumberField(TEXT("shadowQuality"), QualityLevels.ShadowQuality);
+			JsonObject->SetNumberField(TEXT("postProcessQuality"), QualityLevels.PostProcessQuality);
+			JsonObject->SetNumberField(TEXT("textureQuality"), QualityLevels.TextureQuality);
+			JsonObject->SetNumberField(TEXT("effectsQuality"), QualityLevels.EffectsQuality);
+			JsonObject->SetNumberField(TEXT("foliageQuality"), QualityLevels.FoliageQuality);
+			//ComparisonOptions
+			JsonObject->SetBoolField(TEXT("bHasComparisonRules"), true);
+			JsonObject->SetNumberField(TEXT("toleranceRed"), 8);
+			JsonObject->SetNumberField(TEXT("toleranceGreen"), 8);
+			JsonObject->SetNumberField(TEXT("toleranceBlue"), 8);
+			JsonObject->SetNumberField(TEXT("toleranceAlpha"), 8);
+			JsonObject->SetNumberField(TEXT("toleranceMinBrightness"), 0);
+			JsonObject->SetNumberField(TEXT("toleranceMaxBrightness"), 255);
+			JsonObject->SetNumberField(TEXT("maximumLocalError"), MaxLocalError);
+			JsonObject->SetNumberField(TEXT("maximumGlobalError"), MaxGlobalError);
+			JsonObject->SetBoolField(TEXT("bIgnoreAntiAliasing"), true);
+			JsonObject->SetBoolField(TEXT("bIgnoreColors"), false);
+
+
+
+			//Serialize to String
+			TSharedRef<TJsonWriter<>> Writer = TJsonWriterFactory<>::Create(&MetadataJsonString);
+			FJsonSerializer::Serialize(JsonObject.ToSharedRef(), Writer);
+			GLog->Log(FString::Printf(TEXT("Gameplay Automation Screenshot Metadata Serialized to %d characters"), MetadataJsonString.Len()));
+
+
+
+			//Output path for both screenshot image and metadata json
+			DeterminedPath = FPaths::AutomationDir() / TEXT("Incoming") / GWorld->GetName() / InScreenshotName / FPlatformProperties::IniPlatformName();
+
+			//we parse some stuff out of HardwareDetailsString and make a pretty folder name
+			FString HardwareDetailsString;
+			FString HardwareDetails = FHardwareInfo::GetHardwareDetailsString();
+
+			FString RHIString;
+			FString RHILookup = NAME_RHI.ToString() + TEXT("=");
+			if (FParse::Value(*HardwareDetails, *RHILookup, RHIString))
+			{
+				HardwareDetailsString = (HardwareDetailsString + TEXT("_")) + RHIString;
+			}
+
+			FString TextureFormatString;
+			FString TextureFormatLookup = NAME_TextureFormat.ToString() + TEXT("=");
+			if (FParse::Value(*HardwareDetails, *TextureFormatLookup, TextureFormatString))
+			{
+				HardwareDetailsString = (HardwareDetailsString + TEXT("_")) + TextureFormatString;
+			}
+
+			FString DeviceTypeString;
+			FString DeviceTypeLookup = NAME_DeviceType.ToString() + TEXT("=");
+			if (FParse::Value(*HardwareDetails, *DeviceTypeLookup, DeviceTypeString))
+			{
+				HardwareDetailsString = (HardwareDetailsString + TEXT("_")) + TextureFormatString;
+			}
+
+			//Also add FeatureLevel from earlier
+			HardwareDetailsString = (HardwareDetailsString + TEXT("_")) + FeatureLevel;
+
+			if (HardwareDetailsString.Len() > 0)
+			{
+				//remove leading "_"
+				HardwareDetailsString = HardwareDetailsString.RightChop(1);
+			}
+
+			//now plop that back onto the path we're building
+			DeterminedPath = DeterminedPath / HardwareDetailsString;
+			DeterminedPath = DeterminedPath / FPlatformMisc::GetDeviceId() + TEXT(".png");
+
+			//Remove as many noisy rendering conditions as we can until the screenshot has been taken
+			AutomationViewExtension = FSceneViewExtensions::NewExtension<FAutomationUtilsGameplayViewExtension>();
+
+			GLog->Log(FString::Printf(TEXT("Determined Path for screenshot \"%s\" to be %s"), *ScreenshotName, *DeterminedPath));
+		}
+	}
+
+	~FAutomationUtilsGameplayAutomationScreenshotInstance()
+	{
+		Unbind();
+	}
+
+public:
+	void HandleScreenshotData(int32 InSizeX, int32 InSizeY, const TArray<FColor>& InImageData)
+	{
+		check(IsInGameThread());
+
+
+		GLog->Log(FString::Printf(TEXT("Gameplay Automation Screenshot \"%s\" taken with size: %d x %d"), *ScreenshotName, InSizeX, InSizeY));
+
+		//create directory if it doesn't exist
+		IFileManager::Get().MakeDirectory(*FPaths::GetPath(DeterminedPath), true);
+
+		//Save Image File
+		TArray<uint8> CompressedBitmap;
+		FImageUtils::CompressImageArray(InSizeX, InSizeY, InImageData, CompressedBitmap);
+		FFileHelper::SaveArrayToFile(CompressedBitmap, *DeterminedPath);
+		GLog->Log(FString::Printf(TEXT("Saved %d bytes of screenshot image to %s"), CompressedBitmap.Num(), *DeterminedPath));
+
+		//Save Metadata Json
+		FString MetadataPath = FPaths::ChangeExtension(DeterminedPath, TEXT("json"));
+		FFileHelper::SaveStringToFile(MetadataJsonString, *MetadataPath, FFileHelper::EEncodingOptions::ForceUTF8WithoutBOM);
+		GLog->Log(FString::Printf(TEXT("Saved %d bytes of metadata json to %s"), MetadataJsonString.Len(), *MetadataPath));
+
+		//remove our rendering options
+		AutomationViewExtension.Reset();
+
+		//deallocate this object, it's done its part
+		delete this;
+	}
+
+	void WorldDestroyed(ULevel* InLevel, UWorld* InWorld)
+	{
+		if (InLevel == nullptr && InWorld == World.Get())
+		{
+			GLog->Log(FString::Printf(TEXT("Screenshot \"%s\" skipped - level was removed from world before we got our screenshot"), *ScreenshotName));
+			//nothing left for this object to do (we missed the timing somehow), so clean up
+			delete this;
+		}
+	}
+
+private:
+	void Unbind()
+	{
+		if (GEngine && GEngine->GameViewport)
+		{
+			GEngine->GameViewport->OnScreenshotCaptured().RemoveAll(this);
+		}
+		FWorldDelegates::LevelRemovedFromWorld.RemoveAll(this);
+	}
+
+private:
+	FString ScreenshotName;
+	FString MetadataJsonString;
+	FString DeterminedPath;
+	TWeakObjectPtr<UWorld> World;
+	TSharedPtr< class FAutomationUtilsGameplayViewExtension, ESPMode::ThreadSafe > AutomationViewExtension;
+};
+
+
+
+
+
+//Blueprint Library
+UAutomationUtilsBlueprintLibrary::UAutomationUtilsBlueprintLibrary(const FObjectInitializer& ObjectInitializer)
+: Super(ObjectInitializer)
+{
+}
+
+void UAutomationUtilsBlueprintLibrary::TakeGameplayAutomationScreenshot(const FString& ScreenshotName, float MaxGlobalError, float MaxLocalError)
+{
+	//Finish Loading Before Screenshot
+	if (!FPlatformProperties::RequiresCookedData())
+	{
+		//Finish Compiling all shaders
+		GShaderCompilingManager->FinishAllCompilation();
+	}
+
+	//Stream in everything
+	IStreamingManager::Get().StreamAllResources(0.0f);
+
+	//Force all mip maps to load
+	UTexture::ForceUpdateTextureStreaming();
+
+	//Allocate automation object
+	//this new may look scary, but the object deletes itself when its screenshot is processed
+	//(same method as in FTest)
+	FAutomationUtilsGameplayAutomationScreenshotInstance* TempObject = new FAutomationUtilsGameplayAutomationScreenshotInstance(ScreenshotName, MaxGlobalError, MaxLocalError);
+
+	//Actually Take Screenshot
+	FScreenshotRequest::RequestScreenshot(ScreenshotName, false, true);
+}
+
+
+
diff --git a/Engine/Plugins/Experimental/AutomationUtils/Source/AutomationUtils/Public/AutomationUtilsBlueprintLibrary.h b/Engine/Plugins/Experimental/AutomationUtils/Source/AutomationUtils/Public/AutomationUtilsBlueprintLibrary.h
new file mode 100644
--- /dev/null
+++ b/Engine/Plugins/Experimental/AutomationUtils/Source/AutomationUtils/Public/AutomationUtilsBlueprintLibrary.h
@@ -0,0 +1,18 @@
+// Copyright 1998-2019 Epic Games, Inc. All Rights Reserved.
+
+#include "CoreMinimal.h"
+#include "UObject/ObjectMacros.h"
+#include "Engine/LatentActionManager.h"
+#include "Kismet/BlueprintFunctionLibrary.h"
+
+
+#include "AutomationUtilsBlueprintLibrary.generated.h"
+
+UCLASS()
+class AUTOMATIONUTILS_API UAutomationUtilsBlueprintLibrary : public UBlueprintFunctionLibrary
+{
+	GENERATED_UCLASS_BODY()
+
+	UFUNCTION(BlueprintCallable, Category = "Automation")
+	static void TakeGameplayAutomationScreenshot(const FString& ScreenshotName, float MaxGlobalError = .02, float MaxLocalError = .12);
+};
\ No newline at end of file
diff --git a/Engine/Plugins/Experimental/AutomationUtils/Source/AutomationUtilsEditor/AutomationUtilsEditor.Build.cs b/Engine/Plugins/Experimental/AutomationUtils/Source/AutomationUtilsEditor/AutomationUtilsEditor.Build.cs
new file mode 100644
--- /dev/null
+++ b/Engine/Plugins/Experimental/AutomationUtils/Source/AutomationUtilsEditor/AutomationUtilsEditor.Build.cs
@@ -0,0 +1,20 @@
+// Copyright 1998-2019 Epic Games, Inc. All Rights Reserved.
+
+namespace UnrealBuildTool.Rules
+{
+	public class AutomationUtilsEditor : ModuleRules
+	{
+		public AutomationUtilsEditor(ReadOnlyTargetRules Target) : base(Target)
+		{
+			PrivateDependencyModuleNames.AddRange(
+				new string[]
+				{
+					"Core",
+					"CoreUObject",
+					"Engine",
+					"ScreenshotTools"
+				}
+			);
+		}
+	}
+}
diff --git a/Engine/Plugins/Experimental/AutomationUtils/Source/AutomationUtilsEditor/Private/AutomationUtilsEditor.cpp b/Engine/Plugins/Experimental/AutomationUtils/Source/AutomationUtilsEditor/Private/AutomationUtilsEditor.cpp
new file mode 100644
--- /dev/null
+++ b/Engine/Plugins/Experimental/AutomationUtils/Source/AutomationUtilsEditor/Private/AutomationUtilsEditor.cpp
@@ -0,0 +1,26 @@
+// Copyright 1998-2019 Epic Games, Inc. All Rights Reserved.
+
+#include "CoreMinimal.h"
+#include "Modules/ModuleInterface.h"
+#include "Modules/ModuleManager.h"
+
+#define LOCTEXT_NAMESPACE "AutomationUtilsEditor"
+
+class FAutomationUtilsEditorModule : public IModuleInterface
+{
+public:
+
+	virtual void StartupModule() override
+	{
+
+	}
+	
+	virtual void ShutdownModule() override
+	{
+
+	}
+};
+
+IMPLEMENT_MODULE( FAutomationUtilsEditorModule, AutomationUtilsEditor);
+
+#undef LOCTEXT_NAMESPACE
diff --git a/Engine/Plugins/Experimental/AutomationUtils/Source/AutomationUtilsEditor/Private/ScreenshotComparisonCommandlet.cpp b/Engine/Plugins/Experimental/AutomationUtils/Source/AutomationUtilsEditor/Private/ScreenshotComparisonCommandlet.cpp
new file mode 100644
--- /dev/null
+++ b/Engine/Plugins/Experimental/AutomationUtils/Source/AutomationUtilsEditor/Private/ScreenshotComparisonCommandlet.cpp
@@ -0,0 +1,109 @@
+#include "ScreenshotComparisonCommandlet.h"
+#include "Misc/Paths.h"
+#include "Modules/ModuleManager.h"
+#include "Interfaces/IScreenShotToolsModule.h"
+#include "Interfaces/IScreenShotManager.h"
+#include "HAL/FileManager.h"
+
+
+DEFINE_LOG_CATEGORY(LogScreenshotComparison);
+
+UScreenshotComparisonCommandlet::UScreenshotComparisonCommandlet(const FObjectInitializer& ObjectInitializer)
+: Super(ObjectInitializer)
+{
+}
+
+int32 UScreenshotComparisonCommandlet::Main(const FString& CmdLineParameters)
+{
+	FString IncomingPath = FPaths::AutomationDir() / TEXT("Incoming/");
+	FString GroundTruthPath = FPaths::ProjectDir() / TEXT("Test") / TEXT("Screenshots/");
+
+	TArray<FString> MapList;
+	FString MapListStr;
+	if (FParse::Value(*CmdLineParameters, TEXT("Maps="), MapListStr))
+	{
+		MapListStr.ParseIntoArray(MapList, TEXT("+"), true);
+	}
+
+	UE_LOG(LogScreenshotComparison, Log, TEXT("Incoming Path: %s"), *IncomingPath);
+	UE_LOG(LogScreenshotComparison, Log, TEXT("Ground Truth Path: %s"), *GroundTruthPath);
+
+	//Get All Screenshots in Incoming
+	TArray<FString> FilesToCompare;
+	if (FPaths::DirectoryExists(IncomingPath))
+	{
+		FString AbsoluteIncomingPath = IncomingPath;
+		FPaths::MakePathRelativeTo(AbsoluteIncomingPath, *FPaths::RootDir());
+		AbsoluteIncomingPath = FPaths::RootDir() / AbsoluteIncomingPath;
+
+		
+		if (MapList.Num())
+		{
+			for (FString LevelName : MapList)
+			{
+				TArray<FString> LevelFiles;
+				IFileManager::Get().FindFilesRecursive(LevelFiles, *(AbsoluteIncomingPath / LevelName), TEXT("*.png"), true, false);
+				FilesToCompare.Append(LevelFiles);
+			}
+		}
+		else
+		{
+			//No Restrictions -> get ALL images
+			IFileManager::Get().FindFilesRecursive(FilesToCompare, *AbsoluteIncomingPath, TEXT("*.png"), true, false);
+		}
+	}
+
+	if (FilesToCompare.Num())
+	{
+		UE_LOG(LogScreenshotComparison, Log, TEXT("Comparing %d screenshots"), FilesToCompare.Num());
+	}
+	else
+	{
+		UE_LOG(LogScreenshotComparison, Warning, TEXT("Found no screenshots (*.png) in IncomingPath %s"), *IncomingPath);
+	}
+
+
+
+	//Do comparison
+	IScreenShotToolsModule& ScreenShotModule = FModuleManager::LoadModuleChecked<IScreenShotToolsModule>("ScreenShotComparisonTools");
+	IScreenShotManagerPtr ScreenshotManager = ScreenShotModule.GetScreenShotManager();
+	int32 New = 0;
+	int32 Fails = 0;
+	int32 Passes = 0;
+	for (FString ScreenshotName : FilesToCompare)
+	{
+		FPaths::MakePathRelativeTo(ScreenshotName, *IncomingPath);
+		FImageComparisonResult Result = ScreenshotManager->CompareScreenshotAsync(ScreenshotName).Get();
+		if (Result.IsNew())
+		{
+			UE_LOG(LogScreenshotComparison, Warning, TEXT("Incoming file %s is new"), *Result.IncomingFile);
+			New++;
+		}
+		else if (Result.AreSimilar())
+		{
+			UE_LOG(LogScreenshotComparison, Log, TEXT("Incoming file %s is similar!  (Global %f, Local %f)"),
+				*Result.IncomingFile, Result.GlobalDifference, Result.MaxLocalDifference);
+			Passes++;
+		}
+		else
+		{
+			UE_LOG(LogScreenshotComparison, Error, TEXT("Incoming file %s is different! (Global %f, Local %f) : %s"),
+				*Result.IncomingFile, Result.GlobalDifference, Result.MaxLocalDifference, *Result.ErrorMessage.ToString());
+			Fails++;
+		}
+	}
+
+	int32 Total = New + Passes + Fails;
+	if (Total)
+	{
+		UE_LOG(LogScreenshotComparison, Log, TEXT("Comparison Complete!  (New: %d (%.2f%%), Fail: %d (%.2f%%), Pass: %d (%.2f%%)"),
+			New, (float)New / Total * 100, Fails, (float)Fails / Total * 100, Passes, (float)Passes / Total * 100);
+	}
+	else
+	{
+		UE_LOG(LogScreenshotComparison, Log, TEXT("Comparison Complete!"));
+	}
+
+
+	return 0;
+}
\ No newline at end of file
diff --git a/Engine/Plugins/Experimental/AutomationUtils/Source/AutomationUtilsEditor/Public/ScreenshotComparisonCommandlet.h b/Engine/Plugins/Experimental/AutomationUtils/Source/AutomationUtilsEditor/Public/ScreenshotComparisonCommandlet.h
new file mode 100644
--- /dev/null
+++ b/Engine/Plugins/Experimental/AutomationUtils/Source/AutomationUtilsEditor/Public/ScreenshotComparisonCommandlet.h
@@ -0,0 +1,21 @@
+// Copyright 1998-2019 Epic Games, Inc. All Rights Reserved.
+#pragma once
+
+#include "CoreMinimal.h"
+#include "UObject/ObjectMacros.h"
+#include "Commandlets/Commandlet.h"
+#include "ScreenshotComparisonCommandlet.generated.h"
+
+DECLARE_LOG_CATEGORY_EXTERN(LogScreenshotComparison, Log, All);
+
+UCLASS(config = Editor)
+class UScreenshotComparisonCommandlet : public UCommandlet
+{
+	GENERATED_UCLASS_BODY()
+
+public:
+
+	//~ Begin UCommandlet Interface
+	virtual int32 Main(const FString& CmdLineParams) override;
+	//~ End UCommandlet Interface
+};
\ No newline at end of file
