Fixing missing explicit header files in AutomationUtilsBlueprintLibrary.cpp
#rnx

diff --git a/Engine/Plugins/Experimental/AutomationUtils/Source/AutomationUtils/Private/AutomationUtilsBlueprintLibrary.cpp b/Engine/Plugins/Experimental/AutomationUtils/Source/AutomationUtils/Private/AutomationUtilsBlueprintLibrary.cpp
--- a/Engine/Plugins/Experimental/AutomationUtils/Source/AutomationUtils/Private/AutomationUtilsBlueprintLibrary.cpp
+++ b/Engine/Plugins/Experimental/AutomationUtils/Source/AutomationUtils/Private/AutomationUtilsBlueprintLibrary.cpp
@@ -12,6 +12,9 @@
 #include "Materials/MaterialInstanceDynamic.h"
 #include "MeshDrawShaderBindings.h"
 #include "Scalability.h"
+#include "Dom/JsonObject.h"
+#include "Engine/Engine.h"
+#include "Misc/FileHelper.h"
 
 
 
