Fixing missing plugin dependencies on AutomationUtils plugin
#rnx

diff --git a/Engine/Plugins/Experimental/AutomationUtils/AutomationUtils.uplugin b/Engine/Plugins/Experimental/AutomationUtils/AutomationUtils.uplugin
--- a/Engine/Plugins/Experimental/AutomationUtils/AutomationUtils.uplugin
+++ b/Engine/Plugins/Experimental/AutomationUtils/AutomationUtils.uplugin
@@ -1,29 +1,34 @@
 {
-	"FileVersion" : 1,
-	"Version" : 1,
-	"VersionName" : "0.1",
-	"FriendlyName" : "Automation Utilities",
-	"Description" : "Tools and Utilities for Automation purposes",
-	"Category" : "Automation",
-	"CreatedBy" : "Epic Games, Inc.",
-	"CreatedByURL" : "http://epicgames.com",
-	"MarketplaceURL" : "",
-	"SupportURL" : "",
-	"EnabledByDefault" : true,
-	"CanContainContent" : true,
-	"IsBetaVersion" : true,
-	"Installed" : false,
-	"Modules" :
-	[
+	"FileVersion": 1,
+	"Version": 1,
+	"VersionName": "0.1",
+	"FriendlyName": "Automation Utilities",
+	"Description": "Tools and Utilities for Automation purposes",
+	"Category": "Automation",
+	"CreatedBy": "Epic Games, Inc.",
+	"CreatedByURL": "http://epicgames.com",
+	"MarketplaceURL": "",
+	"SupportURL": "",
+	"EnabledByDefault": true,
+	"CanContainContent": true,
+	"IsBetaVersion": true,
+	"Installed": false,
+	"Modules": [
 		{
-			"Name" : "AutomationUtils",
-			"Type" : "Runtime",
-			"LoadingPhase" : "Default"
+			"Name": "AutomationUtils",
+			"Type": "Runtime",
+			"LoadingPhase": "Default"
 		},
 		{
-			"Name" : "AutomationUtilsEditor",
-			"Type" : "Editor",
-			"LoadingPhase" : "Default"
+			"Name": "AutomationUtilsEditor",
+			"Type": "Editor",
+			"LoadingPhase": "Default"
+		}
+	],
+	"Plugins": [
+		{
+			"Name": "ScreenshotTools",
+			"Enabled": true
 		}
 	]
 }
\ No newline at end of file
