Adding the last few Headers to AutomationUtilsBlueprintLibrary
#rnx

diff --git a/Engine/Plugins/Experimental/AutomationUtils/Source/AutomationUtils/Private/AutomationUtilsBlueprintLibrary.cpp b/Engine/Plugins/Experimental/AutomationUtils/Source/AutomationUtils/Private/AutomationUtilsBlueprintLibrary.cpp
--- a/Engine/Plugins/Experimental/AutomationUtils/Source/AutomationUtils/Private/AutomationUtilsBlueprintLibrary.cpp
+++ b/Engine/Plugins/Experimental/AutomationUtils/Source/AutomationUtils/Private/AutomationUtilsBlueprintLibrary.cpp
@@ -15,6 +15,8 @@
 #include "Dom/JsonObject.h"
 #include "Engine/Engine.h"
 #include "Misc/FileHelper.h"
+#include "Serialization/JsonWriter.h"
+#include "Serialization/JsonSerializer.h"
 
 
 
