Adding Copyright Headers to Screenshot Comparison Commandlet

diff --git a/Engine/Plugins/Experimental/AutomationUtils/Source/AutomationUtilsEditor/Private/ScreenshotComparisonCommandlet.cpp b/Engine/Plugins/Experimental/AutomationUtils/Source/AutomationUtilsEditor/Private/ScreenshotComparisonCommandlet.cpp
--- a/Engine/Plugins/Experimental/AutomationUtils/Source/AutomationUtilsEditor/Private/ScreenshotComparisonCommandlet.cpp
+++ b/Engine/Plugins/Experimental/AutomationUtils/Source/AutomationUtilsEditor/Private/ScreenshotComparisonCommandlet.cpp
@@ -1,3 +1,5 @@
+// Copyright 1998-2019 Epic Games, Inc. All Rights Reserved.
+
 #include "ScreenshotComparisonCommandlet.h"
 #include "Misc/Paths.h"
 #include "Modules/ModuleManager.h"
