AutomationUtils: Adding a static factory for GameplayScreenshotInstances to avoid memory leaks and also properly delegate their creation and deletion
#rnx

diff --git a/Engine/Plugins/Experimental/AutomationUtils/Source/AutomationUtils/Private/AutomationUtilsBlueprintLibrary.cpp b/Engine/Plugins/Experimental/AutomationUtils/Source/AutomationUtils/Private/AutomationUtilsBlueprintLibrary.cpp
--- a/Engine/Plugins/Experimental/AutomationUtils/Source/AutomationUtils/Private/AutomationUtilsBlueprintLibrary.cpp
+++ b/Engine/Plugins/Experimental/AutomationUtils/Source/AutomationUtils/Private/AutomationUtilsBlueprintLibrary.cpp
@@ -17,10 +17,7 @@
 #include "Misc/FileHelper.h"
 #include "Serialization/JsonWriter.h"
 #include "Serialization/JsonSerializer.h"
-
-
-
-
+#include "Containers/Ticker.h"
 
 //Private Helper Class Definitions
 class FAutomationUtilsGameplayViewExtension : public FSceneViewExtensionBase
@@ -66,191 +63,262 @@ public:
 };
 
 
+
 class FAutomationUtilsGameplayAutomationScreenshotInstance
 {
 public:
-	FAutomationUtilsGameplayAutomationScreenshotInstance(FString InScreenshotName, float MaxGlobalError, float MaxLocalError)
-		: ScreenshotName(InScreenshotName)
-		, MetadataJsonString(TEXT("{}"))
-		, DeterminedPath(TEXT(""))
-		, World(GWorld.GetReference())
+	FAutomationUtilsGameplayAutomationScreenshotInstance(FString InScreenshotName, float MaxGlobalError, float MaxLocalError);
+	~FAutomationUtilsGameplayAutomationScreenshotInstance();
+
+public:
+	void HandleScreenshotData(int32 InSizeX, int32 InSizeY, const TArray<FColor>& InImageData);
+	void WorldDestroyed(ULevel* InLevel, UWorld* InWorld);
+
+private:
+	void Unbind();
+
+private:
+	FString ScreenshotName;
+	FString MetadataJsonString;
+	FString DeterminedPath;
+	TWeakObjectPtr<UWorld> World;
+	TSharedPtr< class FAutomationUtilsGameplayViewExtension, ESPMode::ThreadSafe > AutomationViewExtension;
+};
+
+
+
+class FAutomationUtilsGameplayAutomationScreenshotFactory
+{
+public:
+	static bool CreateScreenshotInstance(const FString ScreenshotName, float MaxGlobalError, float MaxLocalError);
+	static bool RequestDeleteScreenshotInstance(const FString ScreenshotName);
+
+private:
+	static TMap<FString, FAutomationUtilsGameplayAutomationScreenshotInstance*> ManagedScreenshotInstances;
+};
+
+
+
+FAutomationUtilsGameplayAutomationScreenshotInstance::FAutomationUtilsGameplayAutomationScreenshotInstance(FString InScreenshotName, float MaxGlobalError, float MaxLocalError)
+	: ScreenshotName(*InScreenshotName)
+	, MetadataJsonString(TEXT("{}"))
+	, DeterminedPath(TEXT(""))
+	, World(GWorld.GetReference())
+{
+	if (GEngine && GEngine->GameViewport)
 	{
-		if (GEngine && GEngine->GameViewport)
+		//FlushRendering
+		FlushRenderingCommands();
+		//hook to the screenshot delegate
+		GEngine->GameViewport->OnScreenshotCaptured().AddRaw(this, &FAutomationUtilsGameplayAutomationScreenshotInstance::HandleScreenshotData);
+		//And a removed-from-world delegate too, just in case
+		FWorldDelegates::LevelRemovedFromWorld.AddRaw(this, &FAutomationUtilsGameplayAutomationScreenshotInstance::WorldDestroyed);
+
+
+
+		//Generate Json Metadata relevant to rendering device, quality settings, and comparison tolerances
+		TSharedPtr<FJsonObject> JsonObject = MakeShareable(new FJsonObject);
+		//General Stuff
+		JsonObject->SetStringField(TEXT("name"), FPaths::MakeValidFileName(ScreenshotName, TEXT('_')));
+		JsonObject->SetStringField(TEXT("context"), GWorld->GetName());
+		JsonObject->SetStringField(TEXT("id"), FGuid::NewGuid().ToString());
+		JsonObject->SetStringField(TEXT("Commit"), FEngineVersion::Current().HasChangelist() ? FString::FromInt(FEngineVersion::Current().GetChangelist()) : FString(TEXT("")));
+		FVector2D ViewportSize;  //Width and Height
+		GEngine->GameViewport->GetViewportSize(ViewportSize);
+		JsonObject->SetNumberField(TEXT("width"), ViewportSize.X);
+		JsonObject->SetNumberField(TEXT("height"), ViewportSize.Y);
+		//RHI
+		JsonObject->SetStringField(TEXT("platform"), FPlatformProperties::IniPlatformName());
+		JsonObject->SetStringField(TEXT("rhi"), FHardwareInfo::GetHardwareInfo(NAME_RHI));
+		FString FeatureLevel;
+		GetFeatureLevelName(GMaxRHIFeatureLevel, FeatureLevel);
+		JsonObject->SetStringField(TEXT("featureLevel"), FeatureLevel);
+		JsonObject->SetBoolField(TEXT("bIsStereo"), GEngine->StereoRenderingDevice.IsValid() ? GEngine->StereoRenderingDevice->IsStereoEnabled() : false);
+		//Vendor
+		JsonObject->SetStringField(TEXT("vendor"), RHIVendorIdToString());
+		JsonObject->SetStringField(TEXT("adapterName"), GRHIAdapterName);
+		JsonObject->SetStringField(TEXT("adapterInternalDriverVersion"), GRHIAdapterInternalDriverVersion);
+		JsonObject->SetStringField(TEXT("adapterUserDriverVersion"), GRHIAdapterUserDriverVersion);
+		JsonObject->SetStringField(TEXT("uniqueDeviceId"), FPlatformMisc::GetDeviceId());
+		//Quality
+		Scalability::FQualityLevels QualityLevels = Scalability::GetQualityLevels();
+		JsonObject->SetNumberField(TEXT("resolutionQuality"), QualityLevels.ResolutionQuality);
+		JsonObject->SetNumberField(TEXT("viewDistanceQuality"), QualityLevels.ViewDistanceQuality);
+		JsonObject->SetNumberField(TEXT("antiAliasingQuality"), QualityLevels.AntiAliasingQuality);
+		JsonObject->SetNumberField(TEXT("shadowQuality"), QualityLevels.ShadowQuality);
+		JsonObject->SetNumberField(TEXT("postProcessQuality"), QualityLevels.PostProcessQuality);
+		JsonObject->SetNumberField(TEXT("textureQuality"), QualityLevels.TextureQuality);
+		JsonObject->SetNumberField(TEXT("effectsQuality"), QualityLevels.EffectsQuality);
+		JsonObject->SetNumberField(TEXT("foliageQuality"), QualityLevels.FoliageQuality);
+		//ComparisonOptions
+		JsonObject->SetBoolField(TEXT("bHasComparisonRules"), true);
+		JsonObject->SetNumberField(TEXT("toleranceRed"), 8);
+		JsonObject->SetNumberField(TEXT("toleranceGreen"), 8);
+		JsonObject->SetNumberField(TEXT("toleranceBlue"), 8);
+		JsonObject->SetNumberField(TEXT("toleranceAlpha"), 8);
+		JsonObject->SetNumberField(TEXT("toleranceMinBrightness"), 0);
+		JsonObject->SetNumberField(TEXT("toleranceMaxBrightness"), 255);
+		JsonObject->SetNumberField(TEXT("maximumLocalError"), MaxLocalError);
+		JsonObject->SetNumberField(TEXT("maximumGlobalError"), MaxGlobalError);
+		JsonObject->SetBoolField(TEXT("bIgnoreAntiAliasing"), true);
+		JsonObject->SetBoolField(TEXT("bIgnoreColors"), false);
+
+
+
+		//Serialize to String
+		TSharedRef<TJsonWriter<>> Writer = TJsonWriterFactory<>::Create(&MetadataJsonString);
+		FJsonSerializer::Serialize(JsonObject.ToSharedRef(), Writer);
+		GLog->Log(FString::Printf(TEXT("Gameplay Automation Screenshot Metadata Serialized to %d characters"), MetadataJsonString.Len()));
+
+
+
+		//Output path for both screenshot image and metadata json
+		DeterminedPath = FPaths::AutomationDir() / TEXT("Incoming") / GWorld->GetName() / ScreenshotName / FPlatformProperties::IniPlatformName();
+
+		//we parse some stuff out of HardwareDetailsString and make a pretty folder name
+		FString HardwareDetailsString;
+		FString HardwareDetails = FHardwareInfo::GetHardwareDetailsString();
+
+		FString RHIString;
+		FString RHILookup = NAME_RHI.ToString() + TEXT("=");
+		if (FParse::Value(*HardwareDetails, *RHILookup, RHIString))
 		{
-			//FlushRendering
-			FlushRenderingCommands();
-			//hook to the screenshot delegate
-			GEngine->GameViewport->OnScreenshotCaptured().AddRaw(this, &FAutomationUtilsGameplayAutomationScreenshotInstance::HandleScreenshotData);
-			//And a removed-from-world delegate too, just in case
-			FWorldDelegates::LevelRemovedFromWorld.AddRaw(this, &FAutomationUtilsGameplayAutomationScreenshotInstance::WorldDestroyed);
-
-
-
-			//Generate Json Metadata relevant to rendering device, quality settings, and comparison tolerances
-			TSharedPtr<FJsonObject> JsonObject = MakeShareable(new FJsonObject);
-			//General Stuff
-			JsonObject->SetStringField(TEXT("name"), FPaths::MakeValidFileName(InScreenshotName, TEXT('_')));
-			JsonObject->SetStringField(TEXT("context"), GWorld->GetName());
-			JsonObject->SetStringField(TEXT("id"), FGuid::NewGuid().ToString());
-			JsonObject->SetStringField(TEXT("Commit"), FEngineVersion::Current().HasChangelist() ? FString::FromInt(FEngineVersion::Current().GetChangelist()) : FString(TEXT("")));
-			FVector2D ViewportSize;  //Width and Height
-			GEngine->GameViewport->GetViewportSize(ViewportSize);
-			JsonObject->SetNumberField(TEXT("width"), ViewportSize.X);
-			JsonObject->SetNumberField(TEXT("height"), ViewportSize.Y);
-			//RHI
-			JsonObject->SetStringField(TEXT("platform"), FPlatformProperties::IniPlatformName());
-			JsonObject->SetStringField(TEXT("rhi"), FHardwareInfo::GetHardwareInfo(NAME_RHI));
-			FString FeatureLevel;
-			GetFeatureLevelName(GMaxRHIFeatureLevel, FeatureLevel);
-			JsonObject->SetStringField(TEXT("featureLevel"), FeatureLevel);
-			JsonObject->SetBoolField(TEXT("bIsStereo"), GEngine->StereoRenderingDevice.IsValid() ? GEngine->StereoRenderingDevice->IsStereoEnabled() : false);
-			//Vendor
-			JsonObject->SetStringField(TEXT("vendor"), RHIVendorIdToString());
-			JsonObject->SetStringField(TEXT("adapterName"), GRHIAdapterName);
-			JsonObject->SetStringField(TEXT("adapterInternalDriverVersion"), GRHIAdapterInternalDriverVersion);
-			JsonObject->SetStringField(TEXT("adapterUserDriverVersion"), GRHIAdapterUserDriverVersion);
-			JsonObject->SetStringField(TEXT("uniqueDeviceId"), FPlatformMisc::GetDeviceId());
-			//Quality
-			Scalability::FQualityLevels QualityLevels = Scalability::GetQualityLevels();
-			JsonObject->SetNumberField(TEXT("resolutionQuality"), QualityLevels.ResolutionQuality);
-			JsonObject->SetNumberField(TEXT("viewDistanceQuality"), QualityLevels.ViewDistanceQuality);
-			JsonObject->SetNumberField(TEXT("antiAliasingQuality"), QualityLevels.AntiAliasingQuality);
-			JsonObject->SetNumberField(TEXT("shadowQuality"), QualityLevels.ShadowQuality);
-			JsonObject->SetNumberField(TEXT("postProcessQuality"), QualityLevels.PostProcessQuality);
-			JsonObject->SetNumberField(TEXT("textureQuality"), QualityLevels.TextureQuality);
-			JsonObject->SetNumberField(TEXT("effectsQuality"), QualityLevels.EffectsQuality);
-			JsonObject->SetNumberField(TEXT("foliageQuality"), QualityLevels.FoliageQuality);
-			//ComparisonOptions
-			JsonObject->SetBoolField(TEXT("bHasComparisonRules"), true);
-			JsonObject->SetNumberField(TEXT("toleranceRed"), 8);
-			JsonObject->SetNumberField(TEXT("toleranceGreen"), 8);
-			JsonObject->SetNumberField(TEXT("toleranceBlue"), 8);
-			JsonObject->SetNumberField(TEXT("toleranceAlpha"), 8);
-			JsonObject->SetNumberField(TEXT("toleranceMinBrightness"), 0);
-			JsonObject->SetNumberField(TEXT("toleranceMaxBrightness"), 255);
-			JsonObject->SetNumberField(TEXT("maximumLocalError"), MaxLocalError);
-			JsonObject->SetNumberField(TEXT("maximumGlobalError"), MaxGlobalError);
-			JsonObject->SetBoolField(TEXT("bIgnoreAntiAliasing"), true);
-			JsonObject->SetBoolField(TEXT("bIgnoreColors"), false);
-
-
-
-			//Serialize to String
-			TSharedRef<TJsonWriter<>> Writer = TJsonWriterFactory<>::Create(&MetadataJsonString);
-			FJsonSerializer::Serialize(JsonObject.ToSharedRef(), Writer);
-			GLog->Log(FString::Printf(TEXT("Gameplay Automation Screenshot Metadata Serialized to %d characters"), MetadataJsonString.Len()));
-
-
-
-			//Output path for both screenshot image and metadata json
-			DeterminedPath = FPaths::AutomationDir() / TEXT("Incoming") / GWorld->GetName() / InScreenshotName / FPlatformProperties::IniPlatformName();
-
-			//we parse some stuff out of HardwareDetailsString and make a pretty folder name
-			FString HardwareDetailsString;
-			FString HardwareDetails = FHardwareInfo::GetHardwareDetailsString();
-
-			FString RHIString;
-			FString RHILookup = NAME_RHI.ToString() + TEXT("=");
-			if (FParse::Value(*HardwareDetails, *RHILookup, RHIString))
-			{
-				HardwareDetailsString = (HardwareDetailsString + TEXT("_")) + RHIString;
-			}
-
-			FString TextureFormatString;
-			FString TextureFormatLookup = NAME_TextureFormat.ToString() + TEXT("=");
-			if (FParse::Value(*HardwareDetails, *TextureFormatLookup, TextureFormatString))
-			{
-				HardwareDetailsString = (HardwareDetailsString + TEXT("_")) + TextureFormatString;
-			}
-
-			FString DeviceTypeString;
-			FString DeviceTypeLookup = NAME_DeviceType.ToString() + TEXT("=");
-			if (FParse::Value(*HardwareDetails, *DeviceTypeLookup, DeviceTypeString))
-			{
-				HardwareDetailsString = (HardwareDetailsString + TEXT("_")) + TextureFormatString;
-			}
-
-			//Also add FeatureLevel from earlier
-			HardwareDetailsString = (HardwareDetailsString + TEXT("_")) + FeatureLevel;
-
-			if (HardwareDetailsString.Len() > 0)
-			{
-				//remove leading "_"
-				HardwareDetailsString = HardwareDetailsString.RightChop(1);
-			}
-
-			//now plop that back onto the path we're building
-			DeterminedPath = DeterminedPath / HardwareDetailsString;
-			DeterminedPath = DeterminedPath / FPlatformMisc::GetDeviceId() + TEXT(".png");
-
-			//Remove as many noisy rendering conditions as we can until the screenshot has been taken
-			AutomationViewExtension = FSceneViewExtensions::NewExtension<FAutomationUtilsGameplayViewExtension>();
-
-			GLog->Log(FString::Printf(TEXT("Determined Path for screenshot \"%s\" to be %s"), *ScreenshotName, *DeterminedPath));
+			HardwareDetailsString = (HardwareDetailsString + TEXT("_")) + RHIString;
+		}
+
+		FString TextureFormatString;
+		FString TextureFormatLookup = NAME_TextureFormat.ToString() + TEXT("=");
+		if (FParse::Value(*HardwareDetails, *TextureFormatLookup, TextureFormatString))
+		{
+			HardwareDetailsString = (HardwareDetailsString + TEXT("_")) + TextureFormatString;
 		}
-	}
 
-	~FAutomationUtilsGameplayAutomationScreenshotInstance()
+		FString DeviceTypeString;
+		FString DeviceTypeLookup = NAME_DeviceType.ToString() + TEXT("=");
+		if (FParse::Value(*HardwareDetails, *DeviceTypeLookup, DeviceTypeString))
+		{
+			HardwareDetailsString = (HardwareDetailsString + TEXT("_")) + TextureFormatString;
+		}
+
+		//Also add FeatureLevel from earlier
+		HardwareDetailsString = (HardwareDetailsString + TEXT("_")) + FeatureLevel;
+
+		if (HardwareDetailsString.Len() > 0)
+		{
+			//remove leading "_"
+			HardwareDetailsString = HardwareDetailsString.RightChop(1);
+		}
+
+		//now plop that back onto the path we're building
+		DeterminedPath = DeterminedPath / HardwareDetailsString;
+		DeterminedPath = DeterminedPath / FPlatformMisc::GetDeviceId() + TEXT(".png");
+
+		//Remove as many noisy rendering conditions as we can until the screenshot has been taken
+		AutomationViewExtension = FSceneViewExtensions::NewExtension<FAutomationUtilsGameplayViewExtension>();
+
+		GLog->Log(FString::Printf(TEXT("Determined Path for screenshot \"%s\" to be %s"), *ScreenshotName, *DeterminedPath));
+	}
+	else
 	{
-		Unbind();
+		FTicker::GetCoreTicker().AddTicker(TEXT("FAutomationUtilsGameplayScreenshotInstanceAutoCleanup"), 0.1f, [this](float) {
+			FAutomationUtilsGameplayAutomationScreenshotFactory::RequestDeleteScreenshotInstance(ScreenshotName);
+			return false;
+		});
 	}
+}
 
-public:
-	void HandleScreenshotData(int32 InSizeX, int32 InSizeY, const TArray<FColor>& InImageData)
-	{
-		check(IsInGameThread());
+FAutomationUtilsGameplayAutomationScreenshotInstance::~FAutomationUtilsGameplayAutomationScreenshotInstance()
+{
+	Unbind();
+}
+
+void FAutomationUtilsGameplayAutomationScreenshotInstance::HandleScreenshotData(int32 InSizeX, int32 InSizeY, const TArray<FColor>& InImageData)
+{
+	check(IsInGameThread());
 
 
-		GLog->Log(FString::Printf(TEXT("Gameplay Automation Screenshot \"%s\" taken with size: %d x %d"), *ScreenshotName, InSizeX, InSizeY));
+	GLog->Log(FString::Printf(TEXT("Gameplay Automation Screenshot \"%s\" taken with size: %d x %d"), *ScreenshotName, InSizeX, InSizeY));
 
-		//create directory if it doesn't exist
-		IFileManager::Get().MakeDirectory(*FPaths::GetPath(DeterminedPath), true);
+	//create directory if it doesn't exist
+	IFileManager::Get().MakeDirectory(*FPaths::GetPath(DeterminedPath), true);
 
-		//Save Image File
-		TArray<uint8> CompressedBitmap;
-		FImageUtils::CompressImageArray(InSizeX, InSizeY, InImageData, CompressedBitmap);
-		FFileHelper::SaveArrayToFile(CompressedBitmap, *DeterminedPath);
-		GLog->Log(FString::Printf(TEXT("Saved %d bytes of screenshot image to %s"), CompressedBitmap.Num(), *DeterminedPath));
+	//Save Image File
+	TArray<uint8> CompressedBitmap;
+	FImageUtils::CompressImageArray(InSizeX, InSizeY, InImageData, CompressedBitmap);
+	FFileHelper::SaveArrayToFile(CompressedBitmap, *DeterminedPath);
+	GLog->Log(FString::Printf(TEXT("Saved %d bytes of screenshot image to %s"), CompressedBitmap.Num(), *DeterminedPath));
 
-		//Save Metadata Json
-		FString MetadataPath = FPaths::ChangeExtension(DeterminedPath, TEXT("json"));
-		FFileHelper::SaveStringToFile(MetadataJsonString, *MetadataPath, FFileHelper::EEncodingOptions::ForceUTF8WithoutBOM);
-		GLog->Log(FString::Printf(TEXT("Saved %d bytes of metadata json to %s"), MetadataJsonString.Len(), *MetadataPath));
+	//Save Metadata Json
+	FString MetadataPath = FPaths::ChangeExtension(DeterminedPath, TEXT("json"));
+	FFileHelper::SaveStringToFile(MetadataJsonString, *MetadataPath, FFileHelper::EEncodingOptions::ForceUTF8WithoutBOM);
+	GLog->Log(FString::Printf(TEXT("Saved %d bytes of metadata json to %s"), MetadataJsonString.Len(), *MetadataPath));
 
-		//remove our rendering options
-		AutomationViewExtension.Reset();
+	//remove our rendering options
+	AutomationViewExtension.Reset();
 
-		//deallocate this object, it's done its part
-		delete this;
+	//deallocate this object, it's done its part
+	FAutomationUtilsGameplayAutomationScreenshotFactory::RequestDeleteScreenshotInstance(ScreenshotName);
+}
+
+void FAutomationUtilsGameplayAutomationScreenshotInstance::WorldDestroyed(ULevel* InLevel, UWorld* InWorld)
+{
+	if (InLevel == nullptr && InWorld == World.Get())
+	{
+		GLog->Log(FString::Printf(TEXT("Screenshot \"%s\" skipped - level was removed from world before we got our screenshot"), *ScreenshotName));
+		//nothing left for this object to do (we missed the timing somehow), so clean up
+		FAutomationUtilsGameplayAutomationScreenshotFactory::RequestDeleteScreenshotInstance(ScreenshotName);
 	}
+}
 
-	void WorldDestroyed(ULevel* InLevel, UWorld* InWorld)
+void FAutomationUtilsGameplayAutomationScreenshotInstance::Unbind()
+{
+	if (GEngine && GEngine->GameViewport)
 	{
-		if (InLevel == nullptr && InWorld == World.Get())
-		{
-			GLog->Log(FString::Printf(TEXT("Screenshot \"%s\" skipped - level was removed from world before we got our screenshot"), *ScreenshotName));
-			//nothing left for this object to do (we missed the timing somehow), so clean up
-			delete this;
-		}
+		GEngine->GameViewport->OnScreenshotCaptured().RemoveAll(this);
 	}
+	FWorldDelegates::LevelRemovedFromWorld.RemoveAll(this);
+}
 
-private:
-	void Unbind()
+
+//FAutomationUtilsGameplayAutomationScreenshotFactory member definitions
+bool FAutomationUtilsGameplayAutomationScreenshotFactory::CreateScreenshotInstance(const FString ScreenshotName, float MaxGlobalError, float MaxLocalError)
+{
+	if (ManagedScreenshotInstances.Contains(ScreenshotName))
 	{
-		if (GEngine && GEngine->GameViewport)
-		{
-			GEngine->GameViewport->OnScreenshotCaptured().RemoveAll(this);
-		}
-		FWorldDelegates::LevelRemovedFromWorld.RemoveAll(this);
+		GLog->Log(FString::Printf(TEXT("Cannot create GameplayScreenshotInstance - One with the name \"%s\" already exists!"), *ScreenshotName));
+		return false;
+	}
+	else if (!(GEngine && GEngine->GameViewport))
+	{
+		//there's no GEngine, or else no game viewport to take a screenshot in
+		GLog->Log(TEXT("Cannot create GameplayScreenshotInstance - either GEngine or GameViewport is null!"));
+		return false;
 	}
+	else
+	{
+		ManagedScreenshotInstances.Add(ScreenshotName, new FAutomationUtilsGameplayAutomationScreenshotInstance(ScreenshotName, MaxGlobalError, MaxLocalError));
+		GLog->Log(FString::Printf(TEXT("Created GameplayScreenshotInstance \"%s\""), *ScreenshotName));
+		return true;
+	}
+}
 
-private:
-	FString ScreenshotName;
-	FString MetadataJsonString;
-	FString DeterminedPath;
-	TWeakObjectPtr<UWorld> World;
-	TSharedPtr< class FAutomationUtilsGameplayViewExtension, ESPMode::ThreadSafe > AutomationViewExtension;
-};
+bool FAutomationUtilsGameplayAutomationScreenshotFactory::RequestDeleteScreenshotInstance(const FString ScreenshotName)
+{
+	if (!ManagedScreenshotInstances.Contains(ScreenshotName))
+	{
+		GLog->Log(ELogVerbosity::Error, FString::Printf(TEXT("Tried to remove GameplayScreenshotInstance \"%s\", but it did not exist!"), *ScreenshotName));
+		return false;
+	}
+	else
+	{
+		delete ManagedScreenshotInstances[ScreenshotName];
+		ManagedScreenshotInstances.Remove(ScreenshotName);
+		GLog->Log(FString::Printf(TEXT("Removed GameplayScreenshotInstance \"%s\""), *ScreenshotName));
+		return true;
+	}
+}
+
+TMap<FString, FAutomationUtilsGameplayAutomationScreenshotInstance*> FAutomationUtilsGameplayAutomationScreenshotFactory::ManagedScreenshotInstances;
 
 
 
@@ -262,7 +330,7 @@ UAutomationUtilsBlueprintLibrary::UAutomationUtilsBlueprintLibrary(const FObject
 {
 }
 
-void UAutomationUtilsBlueprintLibrary::TakeGameplayAutomationScreenshot(const FString& ScreenshotName, float MaxGlobalError, float MaxLocalError)
+void UAutomationUtilsBlueprintLibrary::TakeGameplayAutomationScreenshot(const FString ScreenshotName, float MaxGlobalError, float MaxLocalError)
 {
 	//Finish Loading Before Screenshot
 	if (!FPlatformProperties::RequiresCookedData())
@@ -278,9 +346,10 @@ void UAutomationUtilsBlueprintLibrary::TakeGameplayAutomationScreenshot(const FS
 	UTexture::ForceUpdateTextureStreaming();
 
 	//Allocate automation object
-	//this new may look scary, but the object deletes itself when its screenshot is processed
-	//(same method as in FTest)
-	FAutomationUtilsGameplayAutomationScreenshotInstance* TempObject = new FAutomationUtilsGameplayAutomationScreenshotInstance(ScreenshotName, MaxGlobalError, MaxLocalError);
+	if (!FAutomationUtilsGameplayAutomationScreenshotFactory::CreateScreenshotInstance(ScreenshotName, MaxGlobalError, MaxLocalError))
+	{
+		GLog->Log(ELogVerbosity::Error, FString::Printf(TEXT("Failed to create GameplayAutomationScreenshot Instance for \"%s\"!"), *ScreenshotName));
+	}
 
 	//Actually Take Screenshot
 	FScreenshotRequest::RequestScreenshot(ScreenshotName, false, true);
diff --git a/Engine/Plugins/Experimental/AutomationUtils/Source/AutomationUtils/Public/AutomationUtilsBlueprintLibrary.h b/Engine/Plugins/Experimental/AutomationUtils/Source/AutomationUtils/Public/AutomationUtilsBlueprintLibrary.h
--- a/Engine/Plugins/Experimental/AutomationUtils/Source/AutomationUtils/Public/AutomationUtilsBlueprintLibrary.h
+++ b/Engine/Plugins/Experimental/AutomationUtils/Source/AutomationUtils/Public/AutomationUtilsBlueprintLibrary.h
@@ -14,5 +14,5 @@ class AUTOMATIONUTILS_API UAutomationUtilsBlueprintLibrary : public UBlueprintFu
 	GENERATED_UCLASS_BODY()
 
 	UFUNCTION(BlueprintCallable, Category = "Automation")
-	static void TakeGameplayAutomationScreenshot(const FString& ScreenshotName, float MaxGlobalError = .02, float MaxLocalError = .12);
+	static void TakeGameplayAutomationScreenshot(const FString ScreenshotName, float MaxGlobalError = .02, float MaxLocalError = .12);
 };
\ No newline at end of file
