Fixed some non-unity compile issues

diff --git a/Engine/Plugins/Experimental/AutomationUtils/Source/AutomationUtils/Private/AutomationUtilsBlueprintLibrary.cpp b/Engine/Plugins/Experimental/AutomationUtils/Source/AutomationUtils/Private/AutomationUtilsBlueprintLibrary.cpp
--- a/Engine/Plugins/Experimental/AutomationUtils/Source/AutomationUtils/Private/AutomationUtilsBlueprintLibrary.cpp
+++ b/Engine/Plugins/Experimental/AutomationUtils/Source/AutomationUtils/Private/AutomationUtilsBlueprintLibrary.cpp
@@ -14,8 +14,11 @@
 #include "Scalability.h"
 #include "Dom/JsonObject.h"
 #include "Engine/Engine.h"
+#include "Dom/JsonObject.h"
 #include "Misc/FileHelper.h"
 #include "Serialization/JsonWriter.h"
+#include "HAL/FileManager.h"
+#include "Misc/FileHelper.h"
 #include "Serialization/JsonSerializer.h"
 #include "Containers/Ticker.h"
 
