GameplayAutomationScreenshot: Adding MapNameOverride property to the Take-Screenshot method
#rnx

diff --git a/Engine/Plugins/Experimental/AutomationUtils/Source/AutomationUtils/Private/AutomationUtilsBlueprintLibrary.cpp b/Engine/Plugins/Experimental/AutomationUtils/Source/AutomationUtils/Private/AutomationUtilsBlueprintLibrary.cpp
--- a/Engine/Plugins/Experimental/AutomationUtils/Source/AutomationUtils/Private/AutomationUtilsBlueprintLibrary.cpp
+++ b/Engine/Plugins/Experimental/AutomationUtils/Source/AutomationUtils/Private/AutomationUtilsBlueprintLibrary.cpp
@@ -70,7 +70,7 @@ public:
 class FAutomationUtilsGameplayAutomationScreenshotInstance
 {
 public:
-	FAutomationUtilsGameplayAutomationScreenshotInstance(FString InScreenshotName, float MaxGlobalError, float MaxLocalError);
+	FAutomationUtilsGameplayAutomationScreenshotInstance(FString InScreenshotName, float MaxGlobalError, float MaxLocalError, FString MapNameOverride = TEXT(""));
 	~FAutomationUtilsGameplayAutomationScreenshotInstance();
 
 public:
@@ -84,6 +84,7 @@ private:
 	FString ScreenshotName;
 	FString MetadataJsonString;
 	FString DeterminedPath;
+	FString MapName;
 	TWeakObjectPtr<UWorld> World;
 	TSharedPtr< class FAutomationUtilsGameplayViewExtension, ESPMode::ThreadSafe > AutomationViewExtension;
 };
@@ -93,7 +94,7 @@ private:
 class FAutomationUtilsGameplayAutomationScreenshotFactory
 {
 public:
-	static bool CreateScreenshotInstance(const FString ScreenshotName, float MaxGlobalError, float MaxLocalError);
+	static bool CreateScreenshotInstance(const FString ScreenshotName, float MaxGlobalError, float MaxLocalError, FString MapNameOverride);
 	static bool RequestDeleteScreenshotInstance(const FString ScreenshotName);
 
 private:
@@ -102,10 +103,11 @@ private:
 
 
 
-FAutomationUtilsGameplayAutomationScreenshotInstance::FAutomationUtilsGameplayAutomationScreenshotInstance(FString InScreenshotName, float MaxGlobalError, float MaxLocalError)
+FAutomationUtilsGameplayAutomationScreenshotInstance::FAutomationUtilsGameplayAutomationScreenshotInstance(FString InScreenshotName, float MaxGlobalError, float MaxLocalError, FString MapNameOverride)
 	: ScreenshotName(*InScreenshotName)
 	, MetadataJsonString(TEXT("{}"))
 	, DeterminedPath(TEXT(""))
+	, MapName(MapNameOverride.IsEmpty() ? GWorld->GetName() : MapNameOverride)
 	, World(GWorld.GetReference())
 {
 	if (GEngine && GEngine->GameViewport)
@@ -117,13 +119,11 @@ FAutomationUtilsGameplayAutomationScreenshotInstance::FAutomationUtilsGameplayAu
 		//And a removed-from-world delegate too, just in case
 		FWorldDelegates::LevelRemovedFromWorld.AddRaw(this, &FAutomationUtilsGameplayAutomationScreenshotInstance::WorldDestroyed);
 
-
-
 		//Generate Json Metadata relevant to rendering device, quality settings, and comparison tolerances
 		TSharedPtr<FJsonObject> JsonObject = MakeShareable(new FJsonObject);
 		//General Stuff
 		JsonObject->SetStringField(TEXT("name"), FPaths::MakeValidFileName(ScreenshotName, TEXT('_')));
-		JsonObject->SetStringField(TEXT("context"), GWorld->GetName());
+		JsonObject->SetStringField(TEXT("context"), MapName);
 		JsonObject->SetStringField(TEXT("id"), FGuid::NewGuid().ToString());
 		JsonObject->SetStringField(TEXT("Commit"), FEngineVersion::Current().HasChangelist() ? FString::FromInt(FEngineVersion::Current().GetChangelist()) : FString(TEXT("")));
 		FVector2D ViewportSize;  //Width and Height
@@ -176,7 +176,7 @@ FAutomationUtilsGameplayAutomationScreenshotInstance::FAutomationUtilsGameplayAu
 
 
 		//Output path for both screenshot image and metadata json
-		DeterminedPath = FPaths::AutomationDir() / TEXT("Incoming") / GWorld->GetName() / ScreenshotName / FPlatformProperties::IniPlatformName();
+		DeterminedPath = FPaths::AutomationDir() / TEXT("Incoming") / MapName / ScreenshotName / FPlatformProperties::IniPlatformName();
 
 		//we parse some stuff out of HardwareDetailsString and make a pretty folder name
 		FString HardwareDetailsString;
@@ -284,7 +284,7 @@ void FAutomationUtilsGameplayAutomationScreenshotInstance::Unbind()
 
 
 //FAutomationUtilsGameplayAutomationScreenshotFactory member definitions
-bool FAutomationUtilsGameplayAutomationScreenshotFactory::CreateScreenshotInstance(const FString ScreenshotName, float MaxGlobalError, float MaxLocalError)
+bool FAutomationUtilsGameplayAutomationScreenshotFactory::CreateScreenshotInstance(const FString ScreenshotName, float MaxGlobalError, float MaxLocalError, FString MapNameOverride)
 {
 	if (ManagedScreenshotInstances.Contains(ScreenshotName))
 	{
@@ -299,7 +299,7 @@ bool FAutomationUtilsGameplayAutomationScreenshotFactory::CreateScreenshotInstan
 	}
 	else
 	{
-		ManagedScreenshotInstances.Add(ScreenshotName, new FAutomationUtilsGameplayAutomationScreenshotInstance(ScreenshotName, MaxGlobalError, MaxLocalError));
+		ManagedScreenshotInstances.Add(ScreenshotName, new FAutomationUtilsGameplayAutomationScreenshotInstance(ScreenshotName, MaxGlobalError, MaxLocalError, MapNameOverride));
 		GLog->Log(FString::Printf(TEXT("Created GameplayScreenshotInstance \"%s\""), *ScreenshotName));
 		return true;
 	}
@@ -333,7 +333,7 @@ UAutomationUtilsBlueprintLibrary::UAutomationUtilsBlueprintLibrary(const FObject
 {
 }
 
-void UAutomationUtilsBlueprintLibrary::TakeGameplayAutomationScreenshot(const FString ScreenshotName, float MaxGlobalError, float MaxLocalError)
+void UAutomationUtilsBlueprintLibrary::TakeGameplayAutomationScreenshot(const FString ScreenshotName, float MaxGlobalError, float MaxLocalError, FString MapNameOverride)
 {
 	//Finish Loading Before Screenshot
 	if (!FPlatformProperties::RequiresCookedData())
@@ -349,7 +349,7 @@ void UAutomationUtilsBlueprintLibrary::TakeGameplayAutomationScreenshot(const FS
 	UTexture::ForceUpdateTextureStreaming();
 
 	//Allocate automation object
-	if (!FAutomationUtilsGameplayAutomationScreenshotFactory::CreateScreenshotInstance(ScreenshotName, MaxGlobalError, MaxLocalError))
+	if (!FAutomationUtilsGameplayAutomationScreenshotFactory::CreateScreenshotInstance(ScreenshotName, MaxGlobalError, MaxLocalError, MapNameOverride))
 	{
 		GLog->Log(ELogVerbosity::Error, FString::Printf(TEXT("Failed to create GameplayAutomationScreenshot Instance for \"%s\"!"), *ScreenshotName));
 	}
diff --git a/Engine/Plugins/Experimental/AutomationUtils/Source/AutomationUtils/Public/AutomationUtilsBlueprintLibrary.h b/Engine/Plugins/Experimental/AutomationUtils/Source/AutomationUtils/Public/AutomationUtilsBlueprintLibrary.h
--- a/Engine/Plugins/Experimental/AutomationUtils/Source/AutomationUtils/Public/AutomationUtilsBlueprintLibrary.h
+++ b/Engine/Plugins/Experimental/AutomationUtils/Source/AutomationUtils/Public/AutomationUtilsBlueprintLibrary.h
@@ -14,5 +14,5 @@ class AUTOMATIONUTILS_API UAutomationUtilsBlueprintLibrary : public UBlueprintFu
 	GENERATED_UCLASS_BODY()
 
 	UFUNCTION(BlueprintCallable, Category = "Automation")
-	static void TakeGameplayAutomationScreenshot(const FString ScreenshotName, float MaxGlobalError = .02, float MaxLocalError = .12);
+	static void TakeGameplayAutomationScreenshot(const FString ScreenshotName, float MaxGlobalError = .02, float MaxLocalError = .12, FString MapNameOverride = TEXT(""));
 };
\ No newline at end of file
