Make sure we flush AsyncLoading and LevelStreaming before taking a GameplayAutomationScreenshot

diff --git a/Engine/Plugins/Experimental/AutomationUtils/Source/AutomationUtils/Private/AutomationUtilsBlueprintLibrary.cpp b/Engine/Plugins/Experimental/AutomationUtils/Source/AutomationUtils/Private/AutomationUtilsBlueprintLibrary.cpp
--- a/Engine/Plugins/Experimental/AutomationUtils/Source/AutomationUtils/Private/AutomationUtilsBlueprintLibrary.cpp
+++ b/Engine/Plugins/Experimental/AutomationUtils/Source/AutomationUtils/Private/AutomationUtilsBlueprintLibrary.cpp
@@ -21,6 +21,7 @@
 #include "Misc/FileHelper.h"
 #include "Serialization/JsonSerializer.h"
 #include "Containers/Ticker.h"
+#include "Engine/GameEngine.h"
 
 //Private Helper Class Definitions
 class FAutomationUtilsGameplayViewExtension : public FSceneViewExtensionBase
@@ -343,6 +344,17 @@ void UAutomationUtilsBlueprintLibrary::TakeGameplayAutomationScreenshot(const FS
 		GShaderCompilingManager->FinishAllCompilation();
 	}
 
+	FlushAsyncLoading();
+
+	// Make sure we finish all level streaming
+	if (UGameEngine* GameEngine = Cast<UGameEngine>(GEngine))
+	{
+		if (UWorld* GameWorld = GameEngine->GetGameWorld())
+		{
+			GameWorld->FlushLevelStreaming(EFlushLevelStreamingType::Full);
+		}
+	}
+
 	//Stream in everything
 	IStreamingManager::Get().StreamAllResources(0.0f);
 
