Remove deprecated and empty ScreenshotTools plugins

#rnx

diff --git a/Engine/Plugins/Experimental/AutomationUtils/AutomationUtils.uplugin b/Engine/Plugins/Experimental/AutomationUtils/AutomationUtils.uplugin
--- a/Engine/Plugins/Experimental/AutomationUtils/AutomationUtils.uplugin
+++ b/Engine/Plugins/Experimental/AutomationUtils/AutomationUtils.uplugin
@@ -25,12 +25,6 @@
 			"LoadingPhase": "Default"
 		}
 	],
-	"Plugins": [
-		{
-			"Name": "ScreenshotTools",
-			"Enabled": true
-		}
-	],
 	"DocsURL": "",
 	"IsExperimentalVersion": false
 }
\ No newline at end of file
diff --git a/Engine/Plugins/Experimental/AutomationUtils/Source/AutomationUtilsEditor/AutomationUtilsEditor.Build.cs b/Engine/Plugins/Experimental/AutomationUtils/Source/AutomationUtilsEditor/AutomationUtilsEditor.Build.cs
--- a/Engine/Plugins/Experimental/AutomationUtils/Source/AutomationUtilsEditor/AutomationUtilsEditor.Build.cs
+++ b/Engine/Plugins/Experimental/AutomationUtils/Source/AutomationUtilsEditor/AutomationUtilsEditor.Build.cs
@@ -12,7 +12,7 @@ namespace UnrealBuildTool.Rules
 					"Core",
 					"CoreUObject",
 					"Engine",
-					"ScreenshotTools",
+					"ScreenShotComparisonTools",
 					"DesktopPlatform",
 					"AutomationMessages",
 					"Json",
diff --git a/Engine/Plugins/Tests/ScreenshotTools/ScreenshotTools.uplugin b/Engine/Plugins/Tests/ScreenshotTools/ScreenshotTools.uplugin
deleted file mode 100644
diff --git a/Engine/Plugins/Tests/ScreenshotTools/Source/ScreenshotTools/Private/ScreenshotTools.cpp b/Engine/Plugins/Tests/ScreenshotTools/Source/ScreenshotTools/Private/ScreenshotTools.cpp
deleted file mode 100644
diff --git a/Engine/Plugins/Tests/ScreenshotTools/Source/ScreenshotTools/Public/ScreenshotTools.h b/Engine/Plugins/Tests/ScreenshotTools/Source/ScreenshotTools/Public/ScreenshotTools.h
deleted file mode 100644
diff --git a/Engine/Plugins/Tests/ScreenshotTools/Source/ScreenshotTools/ScreenshotTools.Build.cs b/Engine/Plugins/Tests/ScreenshotTools/Source/ScreenshotTools/ScreenshotTools.Build.cs
deleted file mode 100644
